ElectraCodecs: Changed constexpr to a macro due to Clang complaining the value would not be constant

#rnx

#Synced-CL Unknown

diff --git a/Engine/Plugins/Media/NVDECElectra/Source/NVDECElectra/Private/NVDEC/ElectraMediaNVDEC.cpp b/Engine/Plugins/Media/NVDECElectra/Source/NVDECElectra/Private/NVDEC/ElectraMediaNVDEC.cpp
--- a/Engine/Plugins/Media/NVDECElectra/Source/NVDECElectra/Private/NVDEC/ElectraMediaNVDEC.cpp
+++ b/Engine/Plugins/Media/NVDECElectra/Source/NVDECElectra/Private/NVDEC/ElectraMediaNVDEC.cpp
@@ -466,6 +466,8 @@ public:
 };
 
 
+#define MAKE_4CC(a,b,c,d) (((uint32)a << 24) | ((uint32)b << 16) | ((uint32)c << 8) | ((uint32)d))
+
 class FVideoDecoderNVDECElectra : public FNVDECMethods, public IElectraDecoder
 {
 public:
@@ -503,22 +505,22 @@ public:
 		GetFeatures(DecoderFeatures);
 		switch(Codec4CC)
 		{
-			case Make4CC('a','v','c','1'):
-			case Make4CC('a','v','c','3'):
+			case MAKE_4CC('a','v','c','1'):
+			case MAKE_4CC('a','v','c','3'):
 			{
 				return FElectraDecoderBitstreamProcessorH264::Create(DecoderFeatures, InitialCreationOptions);
 			}
-			case Make4CC('h','v','c','1'):
-			case Make4CC('h','e','v','1'):
+			case MAKE_4CC('h','v','c','1'):
+			case MAKE_4CC('h','e','v','1'):
 			{
 				return FElectraDecoderBitstreamProcessorH265::Create(DecoderFeatures, InitialCreationOptions);
 			}
-			case Make4CC('v','p','0','8'):
-			case Make4CC('v','p','0','9'):
+			case MAKE_4CC('v','p','0','8'):
+			case MAKE_4CC('v','p','0','9'):
 			{
 				return FElectraDecoderBitstreamProcessorVPx::Create(DecoderFeatures, InitialCreationOptions);
 			}
-			case Make4CC('a','v','0','1'):
+			case MAKE_4CC('a','v','0','1'):
 			{
 				return FElectraDecoderBitstreamProcessorAV1::Create(DecoderFeatures, InitialCreationOptions);
 			}
@@ -548,11 +550,6 @@ private:
 		Draining
 	};
 
-	static constexpr uint32 Make4CC(const uint8 A, const uint8 B, const uint8 C, const uint8 D)
-	{
-		return (static_cast<uint32>(A) << 24) | (static_cast<uint32>(B) << 16) | (static_cast<uint32>(C) << 8) | static_cast<uint32>(D);
-	}
-
 	bool PostError(int32 ApiReturnValue, FString Message, int32 Code);
 
 	bool InternalDecoderCreate(const FInputAccessUnit& InInputAccessUnit, const TMap<FString, FVariant>& InAdditionalOptions);
@@ -1168,29 +1165,29 @@ bool FVideoDecoderNVDECElectra::InternalDecoderCreate(const FInputAccessUnit& In
 	CUVIDPARSERPARAMS pp {};
 	switch(Codec4CC)
 	{
-		case Make4CC('a','v','c','1'):
-		case Make4CC('a','v','c','3'):
+		case MAKE_4CC('a','v','c','1'):
+		case MAKE_4CC('a','v','c','3'):
 		{
 			pp.CodecType = cudaVideoCodec_H264;
 			break;
 		}
-		case Make4CC('h','v','c','1'):
-		case Make4CC('h','e','v','1'):
+		case MAKE_4CC('h','v','c','1'):
+		case MAKE_4CC('h','e','v','1'):
 		{
 			pp.CodecType = cudaVideoCodec_HEVC;
 			break;
 		}
-		case Make4CC('v','p','0','8'):
+		case MAKE_4CC('v','p','0','8'):
 		{
 			pp.CodecType = cudaVideoCodec_VP8;
 			break;
 		}
-		case Make4CC('v','p','0','9'):
+		case MAKE_4CC('v','p','0','9'):
 		{
 			pp.CodecType = cudaVideoCodec_VP9;
 			break;
 		}
-		case Make4CC('a','v','0','1'):
+		case MAKE_4CC('a','v','0','1'):
 		{
 			pp.CodecType = cudaVideoCodec_AV1;
 			pp.bAnnexb = 0;
@@ -1508,7 +1505,7 @@ FVideoDecoderNVDECElectra::EConvertResult FVideoDecoderNVDECElectra::ConvertDeco
 	{
 		// It is possible that with VP9 we get more than one output corresponding to the same
 		// input access unit (and thus PTS) if the AU contains a superframe.
-		if (Codec4CC == Make4CC('v','p','0','9'))
+		if (Codec4CC == MAKE_4CC('v','p','0','9'))
 		{
 			// We cannot return more than one frame to the upper level and thus need to ignore this frame.
 			return EConvertResult::Success;
